VPSettings: Moved DirectorName and ShowName from WITH_EDITORONLY_DATA block to PUBLIC.
- Added BlueprintReadWrite meta and default values.

diff --git a/Engine/Plugins/Experimental/VirtualProduction/VPSettings/Source/Public/VPSettings.h b/Engine/Plugins/Experimental/VirtualProduction/VPSettings/Source/Public/VPSettings.h
--- a/Engine/Plugins/Experimental/VirtualProduction/VPSettings/Source/Public/VPSettings.h
+++ b/Engine/Plugins/Experimental/VirtualProduction/VPSettings/Source/Public/VPSettings.h
@@ -35,6 +35,14 @@ public:
 	UPROPERTY(EditAnywhere, BlueprintReadWrite, config, Category = "VirtualCamera|Presets")
 	TArray<float> DefaultISOPresets = { 50,100,200,400,800,1600,3200,6400 };
 
+	/** Default Name for Director*/
+	UPROPERTY(config, EditAnywhere, BlueprintReadWrite, Category = "Virtual Production", DisplayName = "Director Name")
+	FString DirectorName = "Default Director";
+
+	/** Default Name for Project*/
+	UPROPERTY(config, EditAnywhere, BlueprintReadWrite, Category = "Virtual Production", DisplayName = "Project Name")
+	FString ShowName = "Default Project";
+
 	UFUNCTION(BlueprintPure, Category = "Virtual Production")
 	static UVPSettings* GetVPSettings();
 
@@ -48,11 +56,6 @@ public:
 	/** Notify when the virtual production roles have changed. */
 	FSimpleMulticastDelegate OnRolesChanged;
 
-	UPROPERTY(config, EditAnywhere, Category = "Virtual Production", DisplayName = "Director Name")
-	FString DirectorName;
-
-	UPROPERTY(config, EditAnywhere, Category = "Virtual Production", DisplayName = "Project Name")
-	FString ShowName;
 #endif
 	
 protected:
