VP Roles - Update deprecation comment

#rnx

diff --git a/Engine/Plugins/Experimental/VirtualProduction/VPSettings/Source/Public/VPSettings.h b/Engine/Plugins/Experimental/VirtualProduction/VPSettings/Source/Public/VPSettings.h
--- a/Engine/Plugins/Experimental/VirtualProduction/VPSettings/Source/Public/VPSettings.h
+++ b/Engine/Plugins/Experimental/VirtualProduction/VPSettings/Source/Public/VPSettings.h
@@ -16,7 +16,7 @@ class VPSETTINGS_API UVPSettings : public UObject
 	GENERATED_BODY()
 
 public:
-	UE_DEPRECATED(5.1, "Use UVirtualProductionRolesSubsystem::GetCurrentVirtualProductionRoles instead.")
+	UE_DEPRECATED(5.1, "Use UVirtualProductionRolesSubsystem::GetCurrentRoles() instead.")
 	const FGameplayTagContainer& GetRoles() const;
 
 	/** Default Kit of Focal Lengths for Virtual Camera */
